Adding Chinese translation to it. #jira ueloc-4711 #fyi jamie.dale, michelle.zhao, alexis.zhang, yali.he

diff --git a/Templates/TP_AEC_ArchvisBP/Config/TemplateDefs.ini b/Templates/TP_AEC_ArchvisBP/Config/TemplateDefs.ini
--- a/Templates/TP_AEC_ArchvisBP/Config/TemplateDefs.ini
+++ b/Templates/TP_AEC_ArchvisBP/Config/TemplateDefs.ini
@@ -16,6 +16,10 @@ LocalizedDescriptions=(Language="en", Text="This template is dedicated to archit
 LocalizedDisplayNames=(Language="ko",Text="건축 시각화")
 LocalizedDescriptions=(Language="ko",Text="이 템플릿은 조명 연구, 인테리어 렌더링 그리고 비실사 스타일라이즈드 렌더링을 위한 예제를 포함한 건축 시각화 워크플로를 위한 템플릿입니다.\n\n이 템플릿은 다양한 디자인 어플리케이션에서 데이터스미스를 이용한 모델 임포트와 시각화가 가능합니다.")
 
+LocalizedDisplayNames=(Language="zh-Hans",Text="建筑可视化")
+LocalizedDescriptions=(Language="zh-Hans",Text="该模板专用于建筑可视化工作流程，包含太阳研究、内饰渲染及非照片级风格化渲染示例。\n\n该模板允许使用Datasmith从多种设计软件中导入模型并轻松为其设置光照。")
+
+
 ; The sort key to use for the template to bypass the default alphabetical sorting.
 SortKey="1"
 
