Updated default plugins and config for TP_AEC_ArchvisBP

diff --git a/Templates/TP_AEC_ArchvisBP/Config/DefaultEngine.ini b/Templates/TP_AEC_ArchvisBP/Config/DefaultEngine.ini
--- a/Templates/TP_AEC_ArchvisBP/Config/DefaultEngine.ini
+++ b/Templates/TP_AEC_ArchvisBP/Config/DefaultEngine.ini
@@ -25,7 +25,15 @@ r.SupportSkyAtmosphereAffectsHeightFog=True
 r.DefaultFeature.LensFlare=True
 r.ClearCoatNormal=False
 r.NormalMapsForStaticLighting=False
-
+r.VirtualTextures=True
+r.VirtualTexturedLightmaps=True
+r.ReflectionMethod=1
+r.GenerateMeshDistanceFields=True
+r.DynamicGlobalIlluminationMethod=1
+r.SkinCache.CompileShaders=True
+r.RayTracing=True
+r.PathTracing=True
+r.RayTracing.Shadows=True
 
 [/Script/EngineSettings.GameMapsSettings]
 ;defines what level to be loaded by default in the editor and at playback time
diff --git a/Templates/TP_AEC_ArchvisBP/TP_AEC_ArchvisBP.uproject b/Templates/TP_AEC_ArchvisBP/TP_AEC_ArchvisBP.uproject
--- a/Templates/TP_AEC_ArchvisBP/TP_AEC_ArchvisBP.uproject
+++ b/Templates/TP_AEC_ArchvisBP/TP_AEC_ArchvisBP.uproject
@@ -72,6 +72,18 @@
 		{
 			"Name": "MovieRenderPipeline",
 			"Enabled": true
+		},
+		{
+			"Name": "MoviePipelineMaskRenderPass",
+			"Enabled": true
+		},
+		{
+			"Name": "MoviePipelineDeadlineIntegration",
+			"Enabled": true
+		},
+		{
+			"Name": "GPULightmass",
+			"Enabled": true
 		}
 	]
 }
\ No newline at end of file
