Tweaked new level selection for TP_AEC_ArchvisBP template

diff --git a/Templates/TP_AEC_ArchvisBP/Config/DefaultEngine.ini b/Templates/TP_AEC_ArchvisBP/Config/DefaultEngine.ini
--- a/Templates/TP_AEC_ArchvisBP/Config/DefaultEngine.ini
+++ b/Templates/TP_AEC_ArchvisBP/Config/DefaultEngine.ini
@@ -48,10 +48,11 @@ EditorStartupMap=/Game/ArchvisProject/Levels/Exterior.Exterior
 GlobalDefaultGameMode=/Game/ArchvisProject/Blueprints/ArchvisGameMode.ArchvisGameMode_C
 
 [/Script/UnrealEd.UnrealEdEngine]
-; Remove Engine Template maps
--TemplateMapInfos=(ThumbnailTexture=Texture2D'/Engine/Maps/Templates/Thumbnails/Default.Default',Map="/Engine/Maps/Templates/Template_Default")
-+TemplateMapInfos=(ThumbnailTexture=Texture2D'/Engine/Maps/Templates/Thumbnails/TimeOfDay.TimeOfDay',Map="/Engine/Maps/Templates/TimeOfDay_Default")
--TemplateMapInfos=(ThumbnailTexture=Texture2D'/Engine/Maps/Templates/Thumbnails/VR-Basic.VR-Basic',Map="/Engine/Maps/Templates/VR-Basic")
+-TemplateMapInfos=(Thumbnail="/Engine/Maps/Templates/Thumbnails/OpenWorld",Map="/Engine/Maps/Templates/OpenWorld")
+-TemplateMapInfos=(Thumbnail="/Engine/Maps/Templates/Thumbnails/Basic",Map="/Engine/Maps/Templates/Template_Default")
+-TemplateMapInfos=(Thumbnail="/Engine/Maps/Templates/Thumbnails/TimeOfDay",Map="/Engine/Maps/Templates/TimeOfDay_Default")
+-TemplateMapInfos=(Thumbnail="/Engine/Maps/Templates/Thumbnails/VR-Basic",Map="/Engine/Maps/Templates/VR-Basic")
++TemplateMapInfos=(Thumbnail="/Engine/Maps/Templates/Thumbnails/Default",Map="/Game/ArchvisProject/Levels/Exterior")
 
 [/Script/Engine.PhysicsSettings]
 DefaultGravityZ=-980.000000
