Updated template definitions for TP_AEC_ArchvisBP #preflight

diff --git a/Templates/TP_AEC_ArchvisBP/Config/TemplateDefs.ini b/Templates/TP_AEC_ArchvisBP/Config/TemplateDefs.ini
--- a/Templates/TP_AEC_ArchvisBP/Config/TemplateDefs.ini
+++ b/Templates/TP_AEC_ArchvisBP/Config/TemplateDefs.ini
@@ -89,10 +89,10 @@ ReplacementsInFiles=(Extensions=("cpp","h","ini","cs"), From="%TEMPLATENAME%", T
 ; GraphicsPreset - Choose between Maximum and Scalable graphics.
 ; StarterContent - Choose whether to add starter content to the project.
 ; VR - Choose whether to enable Virtual Reality support for the project.
-;HiddenSettings=All
-HiddenSettings=Languages
-HiddenSettings=HardwareTarget
-HiddenSettings=GraphicsPreset 
+HiddenSettings=All
+;HiddenSettings=Languages
+;HiddenSettings=HardwareTarget
+;HiddenSettings=GraphicsPreset 
 ;HiddenSettings=StarterContent 
 ;HiddenSettings=XR
 
