Updated new Level entries for TP_AEC_ArchvisBP #preflight

diff --git a/Templates/TP_AEC_ArchvisBP/Config/DefaultEngine.ini b/Templates/TP_AEC_ArchvisBP/Config/DefaultEngine.ini
--- a/Templates/TP_AEC_ArchvisBP/Config/DefaultEngine.ini
+++ b/Templates/TP_AEC_ArchvisBP/Config/DefaultEngine.ini
@@ -53,6 +53,8 @@ GlobalDefaultGameMode=/Game/ArchvisProject/Blueprints/ArchvisGameMode.ArchvisGam
 -TemplateMapInfos=(Thumbnail="/Engine/Maps/Templates/Thumbnails/TimeOfDay",Map="/Engine/Maps/Templates/TimeOfDay_Default")
 -TemplateMapInfos=(Thumbnail="/Engine/Maps/Templates/Thumbnails/VR-Basic",Map="/Engine/Maps/Templates/VR-Basic")
 +TemplateMapInfos=(Thumbnail="/Engine/Maps/Templates/Thumbnails/Default",Map="/Game/ArchvisProject/Levels/Exterior")
++TemplateMapInfos=(Thumbnail="/Engine/Maps/Templates/Thumbnails/Default",Map="/Game/ArchvisProject/Levels/Interior")
++TemplateMapInfos=(Thumbnail="/Engine/Maps/Templates/Thumbnails/Default",Map="/Game/ArchvisProject/Levels/Interior_Lumen")
 
 [/Script/Engine.PhysicsSettings]
 DefaultGravityZ=-980.000000
@@ -158,5 +160,31 @@ ChaosSettings=(DefaultThreadingModel=TaskGraph,DedicatedThreadTickMode=VariableC
 +CollisionChannelRedirects=(OldName="PawnMovement",NewName="Pawn")
 
 [/Script/WindowsTargetPlatform.WindowsTargetSettings]
+Compiler=Default
+-TargetedRHIs=PCD3D_SM5
++TargetedRHIs=PCD3D_SM5
++TargetedRHIs=PCD3D_SM6
 DefaultGraphicsRHI=DefaultGraphicsRHI_DX12
+AudioSampleRate=48000
+AudioCallbackBufferFrameSize=1024
+AudioNumBuffersToEnqueue=1
+AudioMaxChannels=0
+AudioNumSourceWorkers=4
+SpatializationPlugin=
+SourceDataOverridePlugin=
+ReverbPlugin=
+OcclusionPlugin=
+CompressionOverrides=(bOverrideCompressionTimes=False,DurationThreshold=5.000000,MaxNumRandomBranches=0,SoundCueQualityIndex=0)
+bUseAudioStreamCaching=False
+CacheSizeKB=65536
+MaxChunkSizeOverrideKB=0
+bResampleForDevice=False
+MaxSampleRate=48000.000000
+HighSampleRate=32000.000000
+MedSampleRate=24000.000000
+LowSampleRate=12000.000000
+MinSampleRate=8000.000000
+CompressionQualityModifier=1.000000
+AutoStreamingThreshold=0.000000
+SoundCueCookQualityIndex=-1
 
