Updated TP_AEC_ArchvisBP New Level overrides again #preflight

diff --git a/Templates/TP_AEC_ArchvisBP/Config/DefaultEngine.ini b/Templates/TP_AEC_ArchvisBP/Config/DefaultEngine.ini
--- a/Templates/TP_AEC_ArchvisBP/Config/DefaultEngine.ini
+++ b/Templates/TP_AEC_ArchvisBP/Config/DefaultEngine.ini
@@ -42,10 +42,21 @@ r.Lumen.TraceMeshSDF=1
 r.RayTracing.Geometry.InstancedStaticMeshes.Culling=0
 
 [/Script/EngineSettings.GameMapsSettings]
-;defines what level to be loaded by default in the editor and at playback time
-GameDefaultMap=/Game/ArchvisProject/Levels/Exterior.Exterior
 EditorStartupMap=/Game/ArchvisProject/Levels/Exterior.Exterior
++EditorTemplateMapOverrides=(Thumbnail=None,Map="/Game/ArchvisProject/Levels/Exterior.Exterior",DisplayName=NSLOCTEXT("[/Script/EngineSettings]", "73A96829429462B471D06689EC07B2FB", "Exterior"))
++EditorTemplateMapOverrides=(Thumbnail=None,Map="/Game/ArchvisProject/Levels/Interior_Lumen.Interior_Lumen",DisplayName=NSLOCTEXT("[/Script/EngineSettings]", "BDD1166D4B012519C508EE8B74D639B8", "Interior"))
+LocalMapOptions=
+TransitionMap=None
+bUseSplitscreen=True
+TwoPlayerSplitscreenLayout=Horizontal
+ThreePlayerSplitscreenLayout=FavorTop
+FourPlayerSplitscreenLayout=Grid
+bOffsetPlayerGamepadIds=False
+GameInstanceClass=/Script/Engine.GameInstance
+GameDefaultMap=/Game/ArchvisProject/Levels/Exterior.Exterior
+ServerDefaultMap=/Engine/Maps/Entry.Entry
 GlobalDefaultGameMode=/Game/ArchvisProject/Blueprints/ArchvisGameMode.ArchvisGameMode_C
+GlobalDefaultServerGameMode=None
 
 [/Script/UnrealEd.UnrealEdEngine]
 -TemplateMapInfos=(Thumbnail="/Engine/Maps/Templates/Thumbnails/OpenWorld",Map="/Engine/Maps/Templates/OpenWorld")
