Added ArchvisDefault TemplateResources to AEC templates #preflight none

diff --git a/Templates/TP_AEC_ArchvisBP/Config/TemplateDefs.ini b/Templates/TP_AEC_ArchvisBP/Config/TemplateDefs.ini
--- a/Templates/TP_AEC_ArchvisBP/Config/TemplateDefs.ini
+++ b/Templates/TP_AEC_ArchvisBP/Config/TemplateDefs.ini
@@ -106,3 +106,4 @@ HiddenSettings=All
 
 ; Shared feature packs. The files in these packs listed in these structures marked as 'additionalfiles' will be copied on project generation.
 SharedContentPacks=(MountName="Building",DetailLevels=("Standard"))
+SharedContentPacks=(MountName="ArchvisDefault",DetailLevels=("Standard"))
\ No newline at end of file
diff --git a/Templates/TP_AEC_BlankBP/Config/TemplateDefs.ini b/Templates/TP_AEC_BlankBP/Config/TemplateDefs.ini
--- a/Templates/TP_AEC_BlankBP/Config/TemplateDefs.ini
+++ b/Templates/TP_AEC_BlankBP/Config/TemplateDefs.ini
@@ -108,3 +108,4 @@ HiddenSettings=All
 
 ; Shared feature packs. The files in these packs listed in these structures marked as 'additionalfiles' will be copied on project generation.
 ;SharedContentPacks=(MountName="2DSideScroller",DetailLevels=("Standard"))
+SharedContentPacks=(MountName="ArchvisDefault",DetailLevels=("Standard"))
\ No newline at end of file
