Updated EN flavor text for TP_AEC_ArchvisBP template #preflight none

diff --git a/Templates/TP_AEC_ArchvisBP/Config/TemplateDefs.ini b/Templates/TP_AEC_ArchvisBP/Config/TemplateDefs.ini
--- a/Templates/TP_AEC_ArchvisBP/Config/TemplateDefs.ini
+++ b/Templates/TP_AEC_ArchvisBP/Config/TemplateDefs.ini
@@ -11,7 +11,7 @@ Categories=AEC
 LocalizedDisplayNames=(Language="en", Text="Archvis")
 
 ; The localized descriptions for this template, which is displayed when the template is selected.
-LocalizedDescriptions=(Language="en", Text="This template is dedicated to architectural visualization workflows with examples for sun studies, interior renderings, and non-photorealistic stylized renderings.\n\nThis template enables you to use Datasmith to import models from various design applications and easily light them.")
+LocalizedDescriptions=(Language="en", Text="This template is dedicated to architectural visualization workflows. It contains examples of exterior and interior still rendering, animations, sun studies, and non-photorealistic stylized renderings.\n\nThe template gives you a way to get quickly started on your architectural visualization projects by enabling Datasmith and other useful plugins, such as Movie Render Queue and a physically accurate sun and sky system.")
 
 LocalizedDisplayNames=(Language="ko",Text="건축 시각화")
 LocalizedDescriptions=(Language="ko",Text="이 템플릿은 조명 연구, 인테리어 렌더링 그리고 비실사 스타일라이즈드 렌더링을 위한 예제를 포함한 건축 시각화 워크플로를 위한 템플릿입니다.\n\n이 템플릿은 다양한 디자인 어플리케이션에서 데이터스미스를 이용한 모델 임포트와 시각화가 가능합니다.")
