Added Thumbnail overrides to TP_AEC_ArchvisBP and TP_AEC_Blank templates #preflight none

diff --git a/Templates/TP_AEC_ArchvisBP/Config/DefaultEngine.ini b/Templates/TP_AEC_ArchvisBP/Config/DefaultEngine.ini
--- a/Templates/TP_AEC_ArchvisBP/Config/DefaultEngine.ini
+++ b/Templates/TP_AEC_ArchvisBP/Config/DefaultEngine.ini
@@ -43,10 +43,10 @@ r.RayTracing.Geometry.InstancedStaticMeshes.Culling=0
 
 [/Script/EngineSettings.GameMapsSettings]
 EditorStartupMap=/Game/ArchvisProject/Levels/Exterior.Exterior
-+EditorTemplateMapOverrides=(Thumbnail=None,Map="/Game/ArchvisDefault/Levels/Archvis_Empty.Archvis_Empty",DisplayName=NSLOCTEXT("[/Script/EngineSettings]", "BDD1166D4B012519C508EE8B74D639B8", "Archvis Empty"))
-+EditorTemplateMapOverrides=(Thumbnail=None,Map="/Game/ArchvisProject/Levels/Exterior.Exterior",DisplayName=NSLOCTEXT("[/Script/EngineSettings]", "73A96829429462B471D06689EC07B2FB", "Exterior"))
-+EditorTemplateMapOverrides=(Thumbnail=None,Map="/Game/ArchvisProject/Levels/Interior.Interior",DisplayName=NSLOCTEXT("[/Script/EngineSettings]", "2C6771E54254F7715F052A8A5B18AB8C", "Interior"))
-+EditorTemplateMapOverrides=(Thumbnail=None,Map="/Game/ArchvisProject/Levels/Interior_Lightmass.Interior_Lightmass",DisplayName=NSLOCTEXT("[/Script/EngineSettings]", "7D660B88419A2FE50C151FA4408A8D80", "Interior Lightmass"))
++EditorTemplateMapOverrides=(Thumbnail="/Game/ArchvisDefault/Levels/Thumbnails/Thumbnail_Empty.Thumbnail_Empty",Map="/Game/ArchvisDefault/Levels/Archvis_Empty.Archvis_Empty",DisplayName=NSLOCTEXT("[/Script/EngineSettings]", "BDD1166D4B012519C508EE8B74D639B8", "Archvis Empty"))
++EditorTemplateMapOverrides=(Thumbnail="/Game/ArchvisProject/Levels/Thumbnails/Thumbnail_Exterior.Thumbnail_Exterior",Map="/Game/ArchvisProject/Levels/Exterior.Exterior",DisplayName=NSLOCTEXT("[/Script/EngineSettings]", "73A96829429462B471D06689EC07B2FB", "Exterior"))
++EditorTemplateMapOverrides=(Thumbnail="/Game/ArchvisProject/Levels/Thumbnails/Thumbnail_Interior.Thumbnail_Interior",Map="/Game/ArchvisProject/Levels/Interior.Interior",DisplayName=NSLOCTEXT("[/Script/EngineSettings]", "2C6771E54254F7715F052A8A5B18AB8C", "Interior"))
++EditorTemplateMapOverrides=(Thumbnail="/Game/ArchvisProject/Levels/Thumbnails/Thumbnail_InteriorLightmass.Thumbnail_InteriorLightmass",Map="/Game/ArchvisProject/Levels/Interior_Lightmass.Interior_Lightmass",DisplayName=NSLOCTEXT("[/Script/EngineSettings]", "7D660B88419A2FE50C151FA4408A8D80", "Interior Lightmass"))
 LocalMapOptions=
 TransitionMap=None
 bUseSplitscreen=True
diff --git a/Templates/TP_AEC_BlankBP/Config/DefaultEngine.ini b/Templates/TP_AEC_BlankBP/Config/DefaultEngine.ini
--- a/Templates/TP_AEC_BlankBP/Config/DefaultEngine.ini
+++ b/Templates/TP_AEC_BlankBP/Config/DefaultEngine.ini
@@ -42,7 +42,7 @@ r.Lumen.TraceMeshSDF=1
 
 [/Script/EngineSettings.GameMapsSettings]
 EditorStartupMap=/Game/ArchvisDefault/Levels/Archvis_Empty.Archvis_Empty
-+EditorTemplateMapOverrides=(Thumbnail=None,Map="/Game/ArchvisDefault/Levels/Archvis_Empty.Archvis_Empty",DisplayName=NSLOCTEXT("[/Script/EngineSettings]", "E1C1BC224C1F413942879D8305132FF6", "Archvis Empty"))
++EditorTemplateMapOverrides=(Thumbnail="/Game/ArchvisDefault/Levels/Thumbnails/Thumbnail_Empty.Thumbnail_Empty",Map="/Game/ArchvisDefault/Levels/Archvis_Empty.Archvis_Empty",DisplayName=NSLOCTEXT("[/Script/EngineSettings]", "E1C1BC224C1F413942879D8305132FF6", "Archvis Empty"))
 LocalMapOptions=
 TransitionMap=None
 bUseSplitscreen=True
