adding updated loc

diff --git a/Templates/TP_AEC_ArchvisBP/Config/TemplateDefs.ini b/Templates/TP_AEC_ArchvisBP/Config/TemplateDefs.ini
--- a/Templates/TP_AEC_ArchvisBP/Config/TemplateDefs.ini
+++ b/Templates/TP_AEC_ArchvisBP/Config/TemplateDefs.ini
@@ -14,12 +14,12 @@ LocalizedDisplayNames=(Language="en", Text="Archvis")
 LocalizedDescriptions=(Language="en", Text="This template is dedicated to architectural visualization workflows. It contains examples of exterior and interior still rendering, animations, sun studies, and non-photorealistic stylized renderings.\n\nThe template gives you a way to get quickly started on your architectural visualization projects by enabling Datasmith and other useful plugins, such as Movie Render Queue and a physically accurate sun and sky system.")
 
 LocalizedDisplayNames=(Language="ko",Text="건축 시각화")
-LocalizedDescriptions=(Language="ko",Text="이 템플릿은 조명 연구, 인테리어 렌더링 그리고 비실사 스타일라이즈드 렌더링을 위한 예제를 포함한 건축 시각화 워크플로를 위한 템플릿입니다.\n\n이 템플릿은 다양한 디자인 어플리케이션에서 데이터스미스를 이용한 모델 임포트와 시각화가 가능합니다.")
+LocalizedDescriptions=(Language="ko",Text="이 템플릿은 건축 시각화 워크플로를 위해 제작되었습니다. 익스테리어 및 인테리어 스틸 렌더링, 애니메이션, 태양 스터디 및 논 포토리얼한 스타일라이즈드 렌더링 샘플이 포함되어 있습니다.\n\n이 템플릿을 통해 데이터스미스를 비롯해 무비 렌더 큐와 물리적으로 정확한 태양 및 스카이와 같은 다른 유용한 플러그인을 사용하여 건축 시각화 프로젝트를 빠르게 시작하실 수 있습니다.")
 
 LocalizedDisplayNames=(Language="zh-Hans",Text="建筑可视化")
-LocalizedDescriptions=(Language="zh-Hans",Text="该模板专用于建筑可视化工作流程，包含太阳研究、内饰渲染及非照片级风格化渲染示例。\n\n该模板允许使用Datasmith从多种设计软件中导入模型并轻松为其设置光照。")
+LocalizedDescriptions=(Language="zh-Hans",Text="该模板专为建筑可视化流程设计，包含数个室外及室内室内静态图片渲染、动画、阳光研究以及非写实风格化渲染方面的示例。\n\n该模板启用了Datasmith和其他实用插件，如影片渲染队列以及物理精准的太阳和天空系统，可以让你快速启动建筑可视化项目。")
 
-LocalizedDescriptions=(Language="ja",Text="建築ビジュアライゼーションのワークフローに特化したテンプレートです。日照調査やインテリア レンダリング、ノンフォトリアリスティックで様式化されたレンダリングのためのサンプルが備わっています。\n\nこのテンプレートを利用すると、Datasmith を使ってモデルをさまざまなデザイン アプリケーションからインポートできるようになるとともに、それらを簡単にライティングできるようになります。")
+LocalizedDescriptions=(Language="ja",Text="このテンプレートは建築ビジュアライゼーション向けです。エクステリア及びインテリアの静止画レンダリング、アニメーション、日照検討、ノンフォトリアル レンダリングの例を含んでいます。\n\n建築ビジュアライゼーション プロジェクトをすぐに開始できるように準備したテンプレートです。Datasmith に加えて ムービー レンダー キューなど便利なプラグインを有効にし、物理的に正確な太陽、空システムも有効にしています。")
 
 
 ; The sort key to use for the template to bypass the default alphabetical sorting.
