Added main level to AEC Templates and edited startup levels #rb pierrefelix.breton #jira none #preflight none

diff --git a/Templates/TP_AEC_ArchvisBP/Config/DefaultEngine.ini b/Templates/TP_AEC_ArchvisBP/Config/DefaultEngine.ini
--- a/Templates/TP_AEC_ArchvisBP/Config/DefaultEngine.ini
+++ b/Templates/TP_AEC_ArchvisBP/Config/DefaultEngine.ini
@@ -42,7 +42,7 @@ r.Lumen.TraceMeshSDF=1
 r.RayTracing.Geometry.InstancedStaticMeshes.Culling=0
 
 [/Script/EngineSettings.GameMapsSettings]
-EditorStartupMap=/Game/ArchvisProject/Levels/Exterior.Exterior
+EditorStartupMap=/Game/Main.Main
 +EditorTemplateMapOverrides=(Thumbnail="/Game/ArchvisDefault/Levels/Thumbnails/Thumbnail_Empty.Thumbnail_Empty",Map="/Game/ArchvisDefault/Levels/Archvis_Empty.Archvis_Empty",DisplayName=NSLOCTEXT("[/Script/EngineSettings]", "BDD1166D4B012519C508EE8B74D639B8", "Archvis Empty"))
 +EditorTemplateMapOverrides=(Thumbnail="/Game/ArchvisProject/Levels/Thumbnails/Thumbnail_Exterior.Thumbnail_Exterior",Map="/Game/ArchvisProject/Levels/Exterior.Exterior",DisplayName=NSLOCTEXT("[/Script/EngineSettings]", "73A96829429462B471D06689EC07B2FB", "Exterior"))
 +EditorTemplateMapOverrides=(Thumbnail="/Game/ArchvisProject/Levels/Thumbnails/Thumbnail_Interior.Thumbnail_Interior",Map="/Game/ArchvisProject/Levels/Interior.Interior",DisplayName=NSLOCTEXT("[/Script/EngineSettings]", "2C6771E54254F7715F052A8A5B18AB8C", "Interior"))
diff --git a/Templates/TP_AEC_BlankBP/Config/DefaultEngine.ini b/Templates/TP_AEC_BlankBP/Config/DefaultEngine.ini
--- a/Templates/TP_AEC_BlankBP/Config/DefaultEngine.ini
+++ b/Templates/TP_AEC_BlankBP/Config/DefaultEngine.ini
@@ -41,7 +41,7 @@ r.Lumen.HardwareRayTracing=True
 r.Lumen.TraceMeshSDF=1
 
 [/Script/EngineSettings.GameMapsSettings]
-EditorStartupMap=/Game/ArchvisDefault/Levels/Archvis_Empty.Archvis_Empty
+EditorStartupMap=/Game/Main.Main
 +EditorTemplateMapOverrides=(Thumbnail="/Game/ArchvisDefault/Levels/Thumbnails/Thumbnail_Empty.Thumbnail_Empty",Map="/Game/ArchvisDefault/Levels/Archvis_Empty.Archvis_Empty",DisplayName=NSLOCTEXT("[/Script/EngineSettings]", "E1C1BC224C1F413942879D8305132FF6", "Archvis Empty"))
 LocalMapOptions=
 TransitionMap=None
