Turn off SM6 in Archvis template.

diff --git a/Templates/TP_AEC_ArchvisBP/Config/DefaultEngine.ini b/Templates/TP_AEC_ArchvisBP/Config/DefaultEngine.ini
--- a/Templates/TP_AEC_ArchvisBP/Config/DefaultEngine.ini
+++ b/Templates/TP_AEC_ArchvisBP/Config/DefaultEngine.ini
@@ -174,9 +174,6 @@ ChaosSettings=(DefaultThreadingModel=TaskGraph,DedicatedThreadTickMode=VariableC
 
 [/Script/WindowsTargetPlatform.WindowsTargetSettings]
 Compiler=Default
--TargetedRHIs=PCD3D_SM5
-+TargetedRHIs=PCD3D_SM5
-+TargetedRHIs=PCD3D_SM6
 DefaultGraphicsRHI=DefaultGraphicsRHI_DX12
 AudioSampleRate=48000
 AudioCallbackBufferFrameSize=1024
