Added tweaks within DefaultScalability to AEC Templates #preflight

diff --git a/Templates/TP_AEC_ArchvisBP/Config/DefaultScalability.ini b/Templates/TP_AEC_ArchvisBP/Config/DefaultScalability.ini
new file mode 100644
--- /dev/null
+++ b/Templates/TP_AEC_ArchvisBP/Config/DefaultScalability.ini
@@ -0,0 +1,20 @@
+[ScalabilitySettings]
+PerfIndexValues_ResolutionQuality="50 71 75 75 100"
+
+[ShadowQuality@0]
+r.Shadow.Virtual.NonNanite.IncludeInCoarsePages=0
+
+[ShadowQuality@1]
+r.Shadow.Virtual.NonNanite.IncludeInCoarsePages=0
+
+[ShadowQuality@2]
+r.Shadow.Virtual.NonNanite.IncludeInCoarsePages=0
+
+[ShadowQuality@3]
+r.Shadow.Virtual.NonNanite.IncludeInCoarsePages=0
+
+[ShadowQuality@Cine]
+r.Shadow.Virtual.NonNanite.IncludeInCoarsePages=1
+
+[TextureQuality@3]
+r.Streaming.PoolSize=2000
\ No newline at end of file
diff --git a/Templates/TP_AEC_BlankBP/Config/DefaultScalability.ini b/Templates/TP_AEC_BlankBP/Config/DefaultScalability.ini
new file mode 100644
--- /dev/null
+++ b/Templates/TP_AEC_BlankBP/Config/DefaultScalability.ini
@@ -0,0 +1,20 @@
+[ScalabilitySettings]
+PerfIndexValues_ResolutionQuality="50 71 75 75 100"
+
+[ShadowQuality@0]
+r.Shadow.Virtual.NonNanite.IncludeInCoarsePages=0
+
+[ShadowQuality@1]
+r.Shadow.Virtual.NonNanite.IncludeInCoarsePages=0
+
+[ShadowQuality@2]
+r.Shadow.Virtual.NonNanite.IncludeInCoarsePages=0
+
+[ShadowQuality@3]
+r.Shadow.Virtual.NonNanite.IncludeInCoarsePages=0
+
+[ShadowQuality@Cine]
+r.Shadow.Virtual.NonNanite.IncludeInCoarsePages=1
+
+[TextureQuality@3]
+r.Streaming.PoolSize=2000
\ No newline at end of file
