Improved Architecture Template DefaultEngine with MinRoughnessOverride for reflections quality #Jira none #Preflight none

diff --git a/Templates/TP_AEC_ArchvisBP/Config/DefaultEngine.ini b/Templates/TP_AEC_ArchvisBP/Config/DefaultEngine.ini
--- a/Templates/TP_AEC_ArchvisBP/Config/DefaultEngine.ini
+++ b/Templates/TP_AEC_ArchvisBP/Config/DefaultEngine.ini
@@ -41,6 +41,8 @@ r.Lumen.HardwareRayTracing=True
 r.Lumen.TraceMeshSDF=1
 r.RayTracing.Geometry.InstancedStaticMeshes.Culling=0
 
+r.MinRoughnessOverride=0.02
+
 [/Script/EngineSettings.GameMapsSettings]
 EditorStartupMap=/Game/Main.Main
 +EditorTemplateMapOverrides=(Thumbnail="/Game/ArchvisDefault/Levels/Thumbnails/Thumbnail_Empty.Thumbnail_Empty",Map="/Game/ArchvisDefault/Levels/Archvis_Empty.Archvis_Empty",DisplayName=NSLOCTEXT("[/Script/EngineSettings]", "BDD1166D4B012519C508EE8B74D639B8", "Archvis Empty"))
diff --git a/Templates/TP_AEC_BlankBP/Config/DefaultEngine.ini b/Templates/TP_AEC_BlankBP/Config/DefaultEngine.ini
--- a/Templates/TP_AEC_BlankBP/Config/DefaultEngine.ini
+++ b/Templates/TP_AEC_BlankBP/Config/DefaultEngine.ini
@@ -40,6 +40,8 @@ r.RayTracing.Skylight=True
 r.Lumen.HardwareRayTracing=True
 r.Lumen.TraceMeshSDF=1
 
+r.MinRoughnessOverride=0.02
+
 [/Script/EngineSettings.GameMapsSettings]
 EditorStartupMap=/Game/Main.Main
 +EditorTemplateMapOverrides=(Thumbnail="/Game/ArchvisDefault/Levels/Thumbnails/Thumbnail_Empty.Thumbnail_Empty",Map="/Game/ArchvisDefault/Levels/Archvis_Empty.Archvis_Empty",DisplayName=NSLOCTEXT("[/Script/EngineSettings]", "E1C1BC224C1F413942879D8305132FF6", "Archvis Empty"))
