Renamed Shared Ressources Folder from ArchvisDefault to BlankDefault
Updated the change in  TP_AEC_BlankBP and TP_ArchvisBP templates

diff --git a/Templates/TP_AEC_ArchvisBP/Config/DefaultEngine.ini b/Templates/TP_AEC_ArchvisBP/Config/DefaultEngine.ini
--- a/Templates/TP_AEC_ArchvisBP/Config/DefaultEngine.ini
+++ b/Templates/TP_AEC_ArchvisBP/Config/DefaultEngine.ini
@@ -45,7 +45,7 @@ r.MinRoughnessOverride=0.02
 
 [/Script/EngineSettings.GameMapsSettings]
 EditorStartupMap=/Game/Main.Main
-+EditorTemplateMapOverrides=(Thumbnail="/Game/ArchvisDefault/Levels/Thumbnails/Thumbnail_Empty.Thumbnail_Empty",Map="/Game/ArchvisDefault/Levels/Archvis_Empty.Archvis_Empty",DisplayName=NSLOCTEXT("[/Script/EngineSettings]", "BDD1166D4B012519C508EE8B74D639B8", "Archvis Empty"))
++EditorTemplateMapOverrides=(Thumbnail="/Game/BlankDefault/Levels/Thumbnails/Thumbnail_Empty.Thumbnail_Empty",Map="/Game/BlankDefault/Levels/Blank_Empty.Blank_Empty",DisplayName=NSLOCTEXT("[/Script/EngineSettings]", "BDD1166D4B012519C508EE8B74D639B8", "Default Level"))
 +EditorTemplateMapOverrides=(Thumbnail="/Game/ArchvisProject/Levels/Thumbnails/Thumbnail_Exterior.Thumbnail_Exterior",Map="/Game/ArchvisProject/Levels/Exterior.Exterior",DisplayName=NSLOCTEXT("[/Script/EngineSettings]", "73A96829429462B471D06689EC07B2FB", "Exterior"))
 +EditorTemplateMapOverrides=(Thumbnail="/Game/ArchvisProject/Levels/Thumbnails/Thumbnail_Interior.Thumbnail_Interior",Map="/Game/ArchvisProject/Levels/Interior.Interior",DisplayName=NSLOCTEXT("[/Script/EngineSettings]", "2C6771E54254F7715F052A8A5B18AB8C", "Interior"))
 +EditorTemplateMapOverrides=(Thumbnail="/Game/ArchvisProject/Levels/Thumbnails/Thumbnail_InteriorLightmass.Thumbnail_InteriorLightmass",Map="/Game/ArchvisProject/Levels/Interior_Lightmass.Interior_Lightmass",DisplayName=NSLOCTEXT("[/Script/EngineSettings]", "7D660B88419A2FE50C151FA4408A8D80", "Interior Lightmass"))
@@ -202,3 +202,18 @@ MinSampleRate=8000.000000
 CompressionQualityModifier=1.000000
 AutoStreamingThreshold=0.000000
 SoundCueCookQualityIndex=-1
+
+[/Script/AndroidFileServerEditor.AndroidFileServerRuntimeSettings]
+bEnablePlugin=True
+bAllowNetworkConnection=True
+SecurityToken=4F377AFB4FEB1E99A46E548DDE3402D3
+bIncludeInShipping=False
+bAllowExternalStartInShipping=False
+bCompileAFSProject=False
+bUseCompression=False
+bLogFiles=False
+bReportStats=False
+ConnectionType=USBOnly
+bUseManualIPAddress=False
+ManualIPAddress=
+
diff --git a/Templates/TP_AEC_ArchvisBP/Config/TemplateDefs.ini b/Templates/TP_AEC_ArchvisBP/Config/TemplateDefs.ini
--- a/Templates/TP_AEC_ArchvisBP/Config/TemplateDefs.ini
+++ b/Templates/TP_AEC_ArchvisBP/Config/TemplateDefs.ini
@@ -107,4 +107,5 @@ HiddenSettings=All
 
 ; Shared feature packs. The files in these packs listed in these structures marked as 'additionalfiles' will be copied on project generation.
 SharedContentPacks=(MountName="Building",DetailLevels=("Standard"))
-SharedContentPacks=(MountName="ArchvisDefault",DetailLevels=("Standard"))
\ No newline at end of file
+SharedContentPacks=(MountName="BlankDefault",DetailLevels=("Standard"))
+;SharedContentPacks=(MountName="ArchvisDefault",DetailLevels=("Standard"))
\ No newline at end of file
diff --git a/Templates/TP_AEC_BlankBP/Config/DefaultEngine.ini b/Templates/TP_AEC_BlankBP/Config/DefaultEngine.ini
--- a/Templates/TP_AEC_BlankBP/Config/DefaultEngine.ini
+++ b/Templates/TP_AEC_BlankBP/Config/DefaultEngine.ini
@@ -44,7 +44,7 @@ r.MinRoughnessOverride=0.02
 
 [/Script/EngineSettings.GameMapsSettings]
 EditorStartupMap=/Game/Main.Main
-+EditorTemplateMapOverrides=(Thumbnail="/Game/ArchvisDefault/Levels/Thumbnails/Thumbnail_Empty.Thumbnail_Empty",Map="/Game/ArchvisDefault/Levels/Archvis_Empty.Archvis_Empty",DisplayName=NSLOCTEXT("[/Script/EngineSettings]", "E1C1BC224C1F413942879D8305132FF6", "Archvis Empty"))
++EditorTemplateMapOverrides=(Thumbnail="/Game/BlankDefault/Levels/Thumbnails/Thumbnail_Empty.Thumbnail_Empty",Map="/Game/BlankDefault/Levels/Blank_Empty.Blank_Empty",DisplayName=NSLOCTEXT("[/Script/EngineSettings]", "E1C1BC224C1F413942879D8305132FF6", "Default Level"))
 LocalMapOptions=
 TransitionMap=None
 bUseSplitscreen=True
@@ -53,7 +53,7 @@ ThreePlayerSplitscreenLayout=FavorTop
 FourPlayerSplitscreenLayout=Grid
 bOffsetPlayerGamepadIds=False
 GameInstanceClass=/Script/Engine.GameInstance
-GameDefaultMap=/Game/ArchvisDefault/Levels/Archvis_Empty.Archvis_Empty
+GameDefaultMap=/Game/BlankDefault/Levels/Blank_Empty.Blank_Empty
 ServerDefaultMap=/Engine/Maps/Entry.Entry
 GlobalDefaultGameMode=/Script/Engine.GameModeBase
 GlobalDefaultServerGameMode=None
@@ -151,3 +151,18 @@ SoundCueCookQualityIndex=-1
 +CollisionChannelRedirects=(OldName="Dynamic",NewName="WorldDynamic")
 +CollisionChannelRedirects=(OldName="VehicleMovement",NewName="Vehicle")
 +CollisionChannelRedirects=(OldName="PawnMovement",NewName="Pawn")
+
+[/Script/AndroidFileServerEditor.AndroidFileServerRuntimeSettings]
+bEnablePlugin=True
+bAllowNetworkConnection=True
+SecurityToken=5C8924D94D227DCEFACC8BA745A62272
+bIncludeInShipping=False
+bAllowExternalStartInShipping=False
+bCompileAFSProject=False
+bUseCompression=False
+bLogFiles=False
+bReportStats=False
+ConnectionType=USBOnly
+bUseManualIPAddress=False
+ManualIPAddress=
+
diff --git a/Templates/TP_AEC_BlankBP/Config/TemplateDefs.ini b/Templates/TP_AEC_BlankBP/Config/TemplateDefs.ini
--- a/Templates/TP_AEC_BlankBP/Config/TemplateDefs.ini
+++ b/Templates/TP_AEC_BlankBP/Config/TemplateDefs.ini
@@ -109,4 +109,4 @@ HiddenSettings=All
 
 ; Shared feature packs. The files in these packs listed in these structures marked as 'additionalfiles' will be copied on project generation.
 ;SharedContentPacks=(MountName="2DSideScroller",DetailLevels=("Standard"))
-SharedContentPacks=(MountName="ArchvisDefault",DetailLevels=("Standard"))
\ No newline at end of file
+SharedContentPacks=(MountName="BlankDefault",DetailLevels=("Standard"))
\ No newline at end of file
diff --git a/Templates/TemplateResources/Standard/BlankDefault/FeaturePack/manifest.json b/Templates/TemplateResources/Standard/BlankDefault/FeaturePack/manifest.json
new file mode 100644
--- /dev/null
+++ b/Templates/TemplateResources/Standard/BlankDefault/FeaturePack/manifest.json
@@ -0,0 +1,48 @@
+{
+	"Version": 1,
+	"Name":
+	[
+		{
+			"Language": "en",
+			"Text": "Blank templates Shared Resource Standard"
+		}
+	],
+	"Description":
+	[
+		{
+			"Language": "en",
+			"Text": "Shared Resource Pack for templates"
+		}
+	],
+	"AssetTypes":
+	[
+		{
+			"Language": "en",
+			"Text": ""
+		}
+	],
+	"SearchTags":
+	[
+		{
+			"Language": "en",
+			"Text": ""
+		}
+	],	
+	"ClassTypes": " ",
+	"Category": "SharedPack",
+	"Thumbnail": "BlankDefault.png",
+	"Screenshots":
+	[
+		"BlankDefault_Preview.png"
+	],
+	"Ident": "BlankDefaultStandard",
+	"AdditionalFiles":
+	{
+		"DestinationFilesFolder": "BlankDefault",
+		"AdditionalFilesList":
+		[
+			"Templates/TemplateResources/Standard/BlankDefault/Content/*.*"
+		]
+	}
+}
+
