Removing the  r.RayTracing.Skylight option to use the engine default.  This option is no longer encouraged and has side effects with substrate and the material editor as described in https://jira.it.epicgames.com/browse/UE-307547

Changes are low risk as it only affects ini settings in the templates.

diff --git a/Templates/TP_AEC_ArchvisBP/Config/DefaultEngine.ini b/Templates/TP_AEC_ArchvisBP/Config/DefaultEngine.ini
--- a/Templates/TP_AEC_ArchvisBP/Config/DefaultEngine.ini
+++ b/Templates/TP_AEC_ArchvisBP/Config/DefaultEngine.ini
@@ -36,7 +36,6 @@ r.PathTracing=True
 r.RayTracing.Shadows=True
 r.Shadow.Virtual.Enable=1
 r.RayTracing.UseTextureLod=True
-r.RayTracing.Skylight=True
 r.Lumen.HardwareRayTracing=True
 r.Lumen.TraceMeshSDF=1
 r.RayTracing.Geometry.InstancedStaticMeshes.Culling=0
diff --git a/Templates/TP_AEC_BlankBP/Config/DefaultEngine.ini b/Templates/TP_AEC_BlankBP/Config/DefaultEngine.ini
--- a/Templates/TP_AEC_BlankBP/Config/DefaultEngine.ini
+++ b/Templates/TP_AEC_BlankBP/Config/DefaultEngine.ini
@@ -36,7 +36,6 @@ r.PathTracing=True
 r.RayTracing.Shadows=True
 r.Shadow.Virtual.Enable=1
 r.RayTracing.UseTextureLod=True
-r.RayTracing.Skylight=True
 r.Lumen.HardwareRayTracing=True
 r.Lumen.TraceMeshSDF=1
 
